Expose SVGDistanceField module types and functions to Editor Utility Blueprints

#tests Tested in editor
#rnx <Release Notes exclude>

diff --git a/Engine/Source/Editor/SVGDistanceField/Private/SvgDistanceFieldModule.cpp b/Engine/Source/Editor/SVGDistanceField/Private/SvgDistanceFieldModule.cpp
new file mode 100644
--- /dev/null
+++ b/Engine/Source/Editor/SVGDistanceField/Private/SvgDistanceFieldModule.cpp
@@ -0,0 +1,15 @@
+// Copyright Epic Games, Inc. All Rights Reserved.
+
+#include "Modules/ModuleInterface.h"
+#include "Modules/ModuleManager.h"
+
+class FSvgDistanceFieldModule : public IModuleInterface
+{
+public:
+	// IModuleInterface interface
+	virtual void StartupModule() override { }
+	virtual void ShutdownModule() override { }
+
+};
+
+IMPLEMENT_MODULE(FSvgDistanceFieldModule, SVGDistanceField);
diff --git a/Engine/Source/Editor/UnrealEd/Private/EditorEngine.cpp b/Engine/Source/Editor/UnrealEd/Private/EditorEngine.cpp
--- a/Engine/Source/Editor/UnrealEd/Private/EditorEngine.cpp
+++ b/Engine/Source/Editor/UnrealEd/Private/EditorEngine.cpp
@@ -1511,7 +1511,8 @@ void UEditorEngine::LoadDefaultEditorModules()
 			TEXT("UniversalObjectLocatorEditor"),
 			TEXT("StructUtilsEditor"),
 			TEXT("StructUtilsTestSuite"),
-			TEXT("MassEntityEditor")
+			TEXT("MassEntityEditor"),
+			TEXT("SVGDistanceField")
 		};
 
 	FScopedSlowTask ModuleSlowTask((float)UE_ARRAY_COUNT(ModuleNames));
diff --git a/Engine/Source/Runtime/Engine/Engine.Build.cs b/Engine/Source/Runtime/Engine/Engine.Build.cs
--- a/Engine/Source/Runtime/Engine/Engine.Build.cs
+++ b/Engine/Source/Runtime/Engine/Engine.Build.cs
@@ -426,7 +426,8 @@ public class Engine : ModuleRules
 					"HierarchicalLODUtilities",
 					"LevelInstanceEditor",
 					"NaniteBuilder",
-					"MassEntityEditor"
+					"MassEntityEditor",
+					"SVGDistanceField"
 				}
 			);
 		}
